- fixing new plugin dependancy for livelink cameras and redirector

diff --git a/Engine/Plugins/VirtualProduction/LiveLinkCamera/Config/BaseLiveLinkCamera.ini b/Engine/Plugins/VirtualProduction/LiveLinkCamera/Config/BaseLiveLinkCamera.ini
--- a/Engine/Plugins/VirtualProduction/LiveLinkCamera/Config/BaseLiveLinkCamera.ini
+++ b/Engine/Plugins/VirtualProduction/LiveLinkCamera/Config/BaseLiveLinkCamera.ini
@@ -5,4 +5,4 @@
 ;
 
 ; Classes
-+ClassRedirects=(OldName="/Script/LiveLink.LiveLinkCameraController",NewName="/Script/LiveLinkCamera.LiveLinkCameraController")
++ClassRedirects=(OldName="/Script/LiveLinkComponents.LiveLinkCameraController", NewName="/Script/LiveLinkCamera.LiveLinkCameraController")
diff --git a/Templates/TP_InCamVFXBP/TP_InCamVFXBP.uproject b/Templates/TP_InCamVFXBP/TP_InCamVFXBP.uproject
--- a/Templates/TP_InCamVFXBP/TP_InCamVFXBP.uproject
+++ b/Templates/TP_InCamVFXBP/TP_InCamVFXBP.uproject
@@ -38,6 +38,10 @@
 			"Name": "LiveLink",
 			"Enabled": true
 		},
+		{
+			"Name": "LiveLinkCamera",
+			"Enabled": true
+		},
 		{
 			"Name": "MultiUserClient",
 			"Enabled": true
