- Fix CIS, uninitialized variable

diff --git a/Engine/Plugins/VirtualProduction/LiveLinkCamera/Source/LiveLinkCamera/Public/LiveLinkCameraController.h b/Engine/Plugins/VirtualProduction/LiveLinkCamera/Source/LiveLinkCamera/Public/LiveLinkCameraController.h
--- a/Engine/Plugins/VirtualProduction/LiveLinkCamera/Source/LiveLinkCamera/Public/LiveLinkCameraController.h
+++ b/Engine/Plugins/VirtualProduction/LiveLinkCamera/Source/LiveLinkCamera/Public/LiveLinkCameraController.h
@@ -40,7 +40,7 @@ struct LIVELINKCAMERA_API FLensDistortionConfiguration
 	bool bEnableDistortionDebugView = false;
 
 	UPROPERTY(Transient, BlueprintReadOnly, Category = "Live Link")
-	UMaterialInstanceDynamic* LensDistortionMID;
+	UMaterialInstanceDynamic* LensDistortionMID = nullptr;
 };
 
 /**
