Avoid updating the transform for live link cameras with a nodal offset if the position has not changed.

This allows the path tracer to accumulate samples when the camera is not moving.

diff --git a/Engine/Plugins/VirtualProduction/LiveLinkCamera/Source/LiveLinkCamera/Private/LiveLinkCameraController.cpp b/Engine/Plugins/VirtualProduction/LiveLinkCamera/Source/LiveLinkCamera/Private/LiveLinkCameraController.cpp
--- a/Engine/Plugins/VirtualProduction/LiveLinkCamera/Source/LiveLinkCamera/Private/LiveLinkCameraController.cpp
+++ b/Engine/Plugins/VirtualProduction/LiveLinkCamera/Source/LiveLinkCamera/Private/LiveLinkCameraController.cpp
@@ -74,9 +74,6 @@ ULiveLinkCameraController::ULiveLinkCameraController()
 {
 	if (!HasAnyFlags(RF_ArchetypeObject | RF_ClassDefaultObject))
 	{
-		//Hook up to PostActorTick to handle nodal offset
-		FWorldDelegates::OnWorldPostActorTick.AddUObject(this, &ULiveLinkCameraController::OnPostActorTick);
-
 		DistortionProducerID = FGuid::NewGuid();
 	}
 }
@@ -264,11 +261,20 @@ void ULiveLinkCameraController::Cleanup()
 		}
 	}
 
-	FWorldDelegates::OnWorldPostActorTick.RemoveAll(this);
+	if (PostActorTickHandle.IsValid())
+	{
+		FWorldDelegates::OnWorldPostActorTick.Remove(PostActorTickHandle);
+		PostActorTickHandle.Reset();
+	}
 }
 
 void ULiveLinkCameraController::OnEvaluateRegistered()
 {
+	//Hook up to PostActorTick to handle nodal offset
+	if (!PostActorTickHandle.IsValid())
+	{
+		PostActorTickHandle = FWorldDelegates::OnWorldPostActorTick.AddUObject(this, &ULiveLinkCameraController::OnPostActorTick);
+	}
 }
 
 void ULiveLinkCameraController::PostDuplicate(bool bDuplicateForPIE)
@@ -526,8 +532,10 @@ void ULiveLinkCameraController::ApplyNodalOffset(ULensFile* SelectedLensFile, UC
 		const FVector CurrentTranslation = CineCameraComponent->GetRelativeLocation();
 		
 		//Verify if something was set by user / programmatically and update original values
+		bool bWasTransformChanged = false;
 		if (CurrentRotator != LastRotation)
 		{
+			bWasTransformChanged = true;
 			if (CurrentRotator.Pitch != LastRotation.Pitch)
 			{
 				OriginalCameraRotation.Pitch = CurrentRotator.Pitch;
@@ -546,6 +554,7 @@ void ULiveLinkCameraController::ApplyNodalOffset(ULensFile* SelectedLensFile, UC
 
 		if (CurrentTranslation != LastLocation)
 		{
+			bWasTransformChanged = true;
 			if (CurrentTranslation.X != LastLocation.X)
 			{
 				OriginalCameraLocation.X = CurrentTranslation.X;
@@ -568,10 +577,24 @@ void ULiveLinkCameraController::ApplyNodalOffset(ULensFile* SelectedLensFile, UC
 
 			LensFileEvalData.NodalOffset.bWasApplied = true;
 
-			CineCameraComponent->SetRelativeLocation(OriginalCameraLocation);
-			CineCameraComponent->SetRelativeRotation(OriginalCameraRotation.Quaternion());
-			CineCameraComponent->AddLocalOffset(Offset.LocationOffset);
-			CineCameraComponent->AddLocalRotation(Offset.RotationOffset);
+			bWasTransformChanged |= Offset.LocationOffset != LastLocationOffset;
+			bWasTransformChanged |= Offset.RotationOffset != LastRotationOffset;
+
+			if (bWasTransformChanged)
+			{
+				CineCameraComponent->SetRelativeLocation(OriginalCameraLocation);
+				CineCameraComponent->SetRelativeRotation(OriginalCameraRotation.Quaternion());
+				CineCameraComponent->AddLocalOffset(Offset.LocationOffset);
+				CineCameraComponent->AddLocalRotation(Offset.RotationOffset);
+			}
+
+			LastLocationOffset = Offset.LocationOffset;
+			LastRotationOffset = Offset.RotationOffset;
+		}
+		else
+		{
+			LastLocationOffset = FVector::ZeroVector;
+			LastRotationOffset = FQuat(0);
 		}
 
 		LastLocation = CineCameraComponent->GetRelativeLocation();
@@ -581,6 +604,8 @@ void ULiveLinkCameraController::ApplyNodalOffset(ULensFile* SelectedLensFile, UC
 	{
 		LastLocation = FVector::OneVector;
 		LastRotation = FRotator::ZeroRotator;
+		LastLocationOffset = FVector::ZeroVector;
+		LastRotationOffset = FQuat(0);
 	}
 }
 
diff --git a/Engine/Plugins/VirtualProduction/LiveLinkCamera/Source/LiveLinkCamera/Public/LiveLinkCameraController.h b/Engine/Plugins/VirtualProduction/LiveLinkCamera/Source/LiveLinkCamera/Public/LiveLinkCameraController.h
--- a/Engine/Plugins/VirtualProduction/LiveLinkCamera/Source/LiveLinkCamera/Public/LiveLinkCameraController.h
+++ b/Engine/Plugins/VirtualProduction/LiveLinkCamera/Source/LiveLinkCamera/Public/LiveLinkCameraController.h
@@ -185,5 +185,9 @@ private:
 	FCameraFilmbackSettings LastFilmback;
 	FRotator LastRotation;
 	FVector LastLocation;
+	FVector LastLocationOffset;
+	FQuat LastRotationOffset;
 	double LastLensTableVerificationTimestamp = 0.0;
+
+	FDelegateHandle PostActorTickHandle;
 };
\ No newline at end of file
