Fix MRQ renders that are incorrectly rendered without nodal offset applied to the camera.

diff --git a/Engine/Plugins/Animation/LiveLink/Source/LiveLinkComponents/Private/LiveLinkComponentController.cpp b/Engine/Plugins/Animation/LiveLink/Source/LiveLinkComponents/Private/LiveLinkComponentController.cpp
--- a/Engine/Plugins/Animation/LiveLink/Source/LiveLinkComponents/Private/LiveLinkComponentController.cpp
+++ b/Engine/Plugins/Animation/LiveLink/Source/LiveLinkComponents/Private/LiveLinkComponentController.cpp
@@ -185,8 +185,12 @@ void ULiveLinkComponentController::OnRegister()
 #if WITH_EDITOR
 void ULiveLinkComponentController::OnEndPIE(bool bIsSimulating)
 {
-	// Cleanup each controller when PIE session is ending
-	CleanupControllersInMap();
+	const UWorld* const World = GetWorld();
+	if (World && World->WorldType == EWorldType::PIE)
+	{
+		// Cleanup each controller when PIE session is ending
+		CleanupControllersInMap();
+	}
 }
 #endif //WITH_EDITOR
 
diff --git a/Engine/Plugins/VirtualProduction/LiveLinkCamera/Source/LiveLinkCamera/Private/LiveLinkCameraController.cpp b/Engine/Plugins/VirtualProduction/LiveLinkCamera/Source/LiveLinkCamera/Private/LiveLinkCameraController.cpp
--- a/Engine/Plugins/VirtualProduction/LiveLinkCamera/Source/LiveLinkCamera/Private/LiveLinkCameraController.cpp
+++ b/Engine/Plugins/VirtualProduction/LiveLinkCamera/Source/LiveLinkCamera/Private/LiveLinkCameraController.cpp
@@ -3,6 +3,7 @@
 #include "LiveLinkCameraController.h"
 
 #include "Camera/CameraComponent.h"
+#include "Camera/PlayerCameraManager.h"
 #include "CameraCalibrationSubsystem.h"
 #include "CineCameraComponent.h"
 #include "Controllers/LiveLinkTransformController.h"
@@ -10,6 +11,7 @@
 #include "Engine/World.h"
 #include "Features/IModularFeatures.h"
 #include "GameFramework/Actor.h"
+#include "GameFramework/PlayerController.h"
 #include "ILiveLinkClient.h"
 #include "LensFile.h"
 #include "LiveLinkComponentController.h"
@@ -586,6 +588,30 @@ void ULiveLinkCameraController::ApplyNodalOffset(ULensFile* SelectedLensFile, UC
 				CineCameraComponent->SetRelativeRotation(OriginalCameraRotation.Quaternion());
 				CineCameraComponent->AddLocalOffset(Offset.LocationOffset);
 				CineCameraComponent->AddLocalRotation(Offset.RotationOffset);
+
+				// Update the camera manager's cache to reflect this change to nodal offset
+				const UWorld* const World = GetWorld();
+				if (World)
+				{
+					for (FConstPlayerControllerIterator Iterator = World->GetPlayerControllerIterator(); Iterator; ++Iterator)
+					{
+						if (APlayerController* PlayerController = Iterator->Get())
+						{
+							if (APlayerCameraManager* PlayerCameraManager = PlayerController->PlayerCameraManager)
+							{
+								if (PlayerCameraManager->GetViewTarget() == CineCameraComponent->GetOuter())
+								{
+									FMinimalViewInfo CameraViewInfo = PlayerCameraManager->GetCameraCachePOV();
+
+									CameraViewInfo.Location = CineCameraComponent->GetComponentLocation();
+									CameraViewInfo.Rotation = CineCameraComponent->GetComponentRotation();
+
+									PlayerCameraManager->FillCameraCache(CameraViewInfo);
+								}
+							}
+						}
+					}
+				}
 			}
 
 			LastLocationOffset = Offset.LocationOffset;
