[LiveLinkCamera] Fix deprecation warning triggered by GetCameraCachePOV()

diff --git a/Engine/Plugins/VirtualProduction/LiveLinkCamera/Source/LiveLinkCamera/Private/LiveLinkCameraController.cpp b/Engine/Plugins/VirtualProduction/LiveLinkCamera/Source/LiveLinkCamera/Private/LiveLinkCameraController.cpp
--- a/Engine/Plugins/VirtualProduction/LiveLinkCamera/Source/LiveLinkCamera/Private/LiveLinkCameraController.cpp
+++ b/Engine/Plugins/VirtualProduction/LiveLinkCamera/Source/LiveLinkCamera/Private/LiveLinkCameraController.cpp
@@ -601,7 +601,7 @@ void ULiveLinkCameraController::ApplyNodalOffset(ULensFile* SelectedLensFile, UC
 							{
 								if (PlayerCameraManager->GetViewTarget() == CineCameraComponent->GetOuter())
 								{
-									FMinimalViewInfo CameraViewInfo = PlayerCameraManager->GetCameraCachePOV();
+									FMinimalViewInfo CameraViewInfo = PlayerCameraManager->GetCameraCacheView();
 
 									CameraViewInfo.Location = CineCameraComponent->GetComponentLocation();
 									CameraViewInfo.Rotation = CineCameraComponent->GetComponentRotation();
