Add depth of field flag to LiveLink camera role static data to allow a LiveLink source to disable depth of field on the camera

diff --git a/Engine/Plugins/VirtualProduction/LiveLinkCamera/Source/LiveLinkCamera/Private/LiveLinkCameraController.cpp b/Engine/Plugins/VirtualProduction/LiveLinkCamera/Source/LiveLinkCamera/Private/LiveLinkCameraController.cpp
--- a/Engine/Plugins/VirtualProduction/LiveLinkCamera/Source/LiveLinkCamera/Private/LiveLinkCameraController.cpp
+++ b/Engine/Plugins/VirtualProduction/LiveLinkCamera/Source/LiveLinkCamera/Private/LiveLinkCameraController.cpp
@@ -116,6 +116,20 @@ void ULiveLinkCameraController::Tick(float DeltaTime, const FLiveLinkSubjectFram
 
 			if (UCineCameraComponent* CineCameraComponent = Cast<UCineCameraComponent>(CameraComponent))
 			{
+				// If depth of field is not supported, force disable depth of field on the camera
+				// Otherwise, if depth of field is supported and the frame data has a valid value for focus, force the focus method to Manual 
+				if (UpdateFlags.bApplyFocusDistance)
+				{
+					if (StaticData->bIsDepthOfFieldSupported == false)
+					{
+						CineCameraComponent->FocusSettings.FocusMethod = ECameraFocusMethod::Disable;
+					}
+					else if (StaticData->bIsFocusDistanceSupported)
+					{
+						CineCameraComponent->FocusSettings.FocusMethod = ECameraFocusMethod::Manual;
+					}
+				}
+
 				ULensFile* SelectedLensFile = LensFilePicker.GetLensFile();
 				LensFileEvalData.LensFile = SelectedLensFile;
 
diff --git a/Engine/Source/Runtime/LiveLinkInterface/Public/Roles/LiveLinkCameraTypes.h b/Engine/Source/Runtime/LiveLinkInterface/Public/Roles/LiveLinkCameraTypes.h
--- a/Engine/Source/Runtime/LiveLinkInterface/Public/Roles/LiveLinkCameraTypes.h
+++ b/Engine/Source/Runtime/LiveLinkInterface/Public/Roles/LiveLinkCameraTypes.h
@@ -54,6 +54,10 @@ public:
 	//Whether FocusDistance in frame data can be used
 	UPROPERTY(EditAnywhere, BlueprintReadWrite, Category = "LiveLink")
 	bool bIsFocusDistanceSupported = false;
+
+	//Set to false to force the camera to disable depth of field
+	UPROPERTY(EditAnywhere, BlueprintReadWrite, Category = "LiveLink")
+	bool bIsDepthOfFieldSupported = true;
 };
 
 /**
