Add null check on LiveLink camera controller's outer actor

diff --git a/Engine/Plugins/VirtualProduction/LiveLinkCamera/Source/LiveLinkCamera/Private/LiveLinkCameraController.cpp b/Engine/Plugins/VirtualProduction/LiveLinkCamera/Source/LiveLinkCamera/Private/LiveLinkCameraController.cpp
--- a/Engine/Plugins/VirtualProduction/LiveLinkCamera/Source/LiveLinkCamera/Private/LiveLinkCameraController.cpp
+++ b/Engine/Plugins/VirtualProduction/LiveLinkCamera/Source/LiveLinkCamera/Private/LiveLinkCameraController.cpp
@@ -316,73 +316,76 @@ void ULiveLinkCameraController::PostLoad()
 		// we have to migrate some properties from the camera controller to the lens component. If the owning actor does not have already have a lens 
 		// component then we will instance a new one.
 		AActor* const CameraActor = GetOuterActor();
-		CameraActor->ConditionalPostLoad();
-		if (CameraActor->GetWorld())
+		if (ensure(CameraActor))
 		{
-			TInlineComponentArray<ULensComponent*> LensComponents;
-			CameraActor->GetComponents(LensComponents);
-
-			// Loop through any existing lens components to find those that may already have been evaluating the same LensFile as this camera controller
-			bool bFoundMatchingLensComponent = false;
-			for (ULensComponent* LensComponent : LensComponents)
+			CameraActor->ConditionalPostLoad();
+			if (CameraActor->GetWorld())
 			{
-				if (LensComponent->GetLensFile() == LensFilePicker.GetLensFile())
-				{
-					bFoundMatchingLensComponent = true;
-
-					// It is very unlikely that both the camera controller and the existing lens component were both applying nodal offset simultaneously.
-					// Therefore, it is safe to assume that if the camera controller was previously applying nodal offset, the new desired behavior would be for 
-					// the lens component to start applying it. But if the lens component was previously applying nodal offset, then there is no need to update anything.
-					PRAGMA_DISABLE_DEPRECATION_WARNINGS
-					if (bApplyNodalOffset_DEPRECATED)
-					{
-						LensComponent->SetApplyNodalOffsetOnTick(bApplyNodalOffset_DEPRECATED);
-					}
+				TInlineComponentArray<ULensComponent*> LensComponents;
+				CameraActor->GetComponents(LensComponents);
 
-					// Copy the overscan multiplier settings to the lens component
-					if (bScaleOverscan_DEPRECATED)
-					{
-						LensComponent->SetOverscanMultiplier(OverscanMultiplier_DEPRECATED);
-					}
-
-					// Copy the cropped filmback settings
-					if (bUseCroppedFilmback_DEPRECATED)
+				// Loop through any existing lens components to find those that may already have been evaluating the same LensFile as this camera controller
+				bool bFoundMatchingLensComponent = false;
+				for (ULensComponent* LensComponent : LensComponents)
+				{
+					if (LensComponent->GetLensFile() == LensFilePicker.GetLensFile())
 					{
-						LensComponent->SetFilmbackOverrideSetting(EFilmbackOverrideSource::CroppedFilmbackSetting);
-						LensComponent->SetCroppedFilmback(CroppedFilmback_DEPRECATED);
+						bFoundMatchingLensComponent = true;
+
+						// It is very unlikely that both the camera controller and the existing lens component were both applying nodal offset simultaneously.
+						// Therefore, it is safe to assume that if the camera controller was previously applying nodal offset, the new desired behavior would be for 
+						// the lens component to start applying it. But if the lens component was previously applying nodal offset, then there is no need to update anything.
+						PRAGMA_DISABLE_DEPRECATION_WARNINGS
+							if (bApplyNodalOffset_DEPRECATED)
+							{
+								LensComponent->SetApplyNodalOffsetOnTick(bApplyNodalOffset_DEPRECATED);
+							}
+
+						// Copy the overscan multiplier settings to the lens component
+						if (bScaleOverscan_DEPRECATED)
+						{
+							LensComponent->SetOverscanMultiplier(OverscanMultiplier_DEPRECATED);
+						}
+
+						// Copy the cropped filmback settings
+						if (bUseCroppedFilmback_DEPRECATED)
+						{
+							LensComponent->SetFilmbackOverrideSetting(EFilmbackOverrideSource::CroppedFilmbackSetting);
+							LensComponent->SetCroppedFilmback(CroppedFilmback_DEPRECATED);
+						}
+
+						PRAGMA_ENABLE_DEPRECATION_WARNINGS
 					}
-
-					PRAGMA_ENABLE_DEPRECATION_WARNINGS
 				}
-			}
 
-			// If no matching existing lens components were found, create a new one and initialize it with the correct LensFile
-			if (!bFoundMatchingLensComponent)
-			{
-				if (LensFilePicker.GetLensFile())
+				// If no matching existing lens components were found, create a new one and initialize it with the correct LensFile
+				if (!bFoundMatchingLensComponent)
 				{
-					ULensComponent* LensComponent = NewObject<ULensComponent>(CameraActor, MakeUniqueObjectName(this, ULensComponent::StaticClass(), TEXT("Lens")));
-					CameraActor->AddInstanceComponent(LensComponent);
-					LensComponent->RegisterComponent();
-
-					LensComponent->SetLensFilePicker(LensFilePicker);
-
-					PRAGMA_DISABLE_DEPRECATION_WARNINGS
-					LensComponent->SetApplyNodalOffsetOnTick(bApplyNodalOffset_DEPRECATED);
-
-					// Copy the overscan multiplier settings to the lens component
-					if (bScaleOverscan_DEPRECATED)
-					{
-						LensComponent->SetOverscanMultiplier(OverscanMultiplier_DEPRECATED);
-					}
-
-					// Copy the cropped filmback settings
-					if (bUseCroppedFilmback_DEPRECATED)
+					if (LensFilePicker.GetLensFile())
 					{
-						LensComponent->SetFilmbackOverrideSetting(EFilmbackOverrideSource::CroppedFilmbackSetting);
-						LensComponent->SetCroppedFilmback(CroppedFilmback_DEPRECATED);
+						ULensComponent* LensComponent = NewObject<ULensComponent>(CameraActor, MakeUniqueObjectName(this, ULensComponent::StaticClass(), TEXT("Lens")));
+						CameraActor->AddInstanceComponent(LensComponent);
+						LensComponent->RegisterComponent();
+
+						LensComponent->SetLensFilePicker(LensFilePicker);
+
+						PRAGMA_DISABLE_DEPRECATION_WARNINGS
+							LensComponent->SetApplyNodalOffsetOnTick(bApplyNodalOffset_DEPRECATED);
+
+						// Copy the overscan multiplier settings to the lens component
+						if (bScaleOverscan_DEPRECATED)
+						{
+							LensComponent->SetOverscanMultiplier(OverscanMultiplier_DEPRECATED);
+						}
+
+						// Copy the cropped filmback settings
+						if (bUseCroppedFilmback_DEPRECATED)
+						{
+							LensComponent->SetFilmbackOverrideSetting(EFilmbackOverrideSource::CroppedFilmbackSetting);
+							LensComponent->SetCroppedFilmback(CroppedFilmback_DEPRECATED);
+						}
+						PRAGMA_ENABLE_DEPRECATION_WARNINGS
 					}
-					PRAGMA_ENABLE_DEPRECATION_WARNINGS
 				}
 			}
 		}
