Fix lens calibration cooking error occuring in QAVirtualProduction as well as some warnings

diff --git a/Engine/Plugins/VirtualProduction/LiveLinkCamera/Source/LiveLinkCamera/Private/LiveLinkCameraController.cpp b/Engine/Plugins/VirtualProduction/LiveLinkCamera/Source/LiveLinkCamera/Private/LiveLinkCameraController.cpp
--- a/Engine/Plugins/VirtualProduction/LiveLinkCamera/Source/LiveLinkCamera/Private/LiveLinkCameraController.cpp
+++ b/Engine/Plugins/VirtualProduction/LiveLinkCamera/Source/LiveLinkCamera/Private/LiveLinkCameraController.cpp
@@ -313,6 +313,10 @@ void ULiveLinkCameraController::PostLoad()
 	const int32 UE5MainVersion = GetLinkerCustomVersion(FUE5MainStreamObjectVersion::GUID);
 	if (UE5MainVersion < FUE5MainStreamObjectVersion::LensComponentNodalOffset)
 	{
+		if (HasAnyFlags(RF_ArchetypeObject))
+		{
+			return;
+		}
 		// LensFile and distortion evaluation for nodal offset has moved to the Lens component. In order for existing camera actors to continue functioning as before,
 		// we have to migrate some properties from the camera controller to the lens component. If the owning actor does not have already have a lens 
 		// component then we will instance a new one.
