Prevent the LiveLink camera controller from locking the focus method of the controlled camera component
#rnx

diff --git a/Engine/Plugins/VirtualProduction/LiveLinkCamera/Source/LiveLinkCamera/Private/LiveLinkCameraController.cpp b/Engine/Plugins/VirtualProduction/LiveLinkCamera/Source/LiveLinkCamera/Private/LiveLinkCameraController.cpp
--- a/Engine/Plugins/VirtualProduction/LiveLinkCamera/Source/LiveLinkCamera/Private/LiveLinkCameraController.cpp
+++ b/Engine/Plugins/VirtualProduction/LiveLinkCamera/Source/LiveLinkCamera/Private/LiveLinkCameraController.cpp
@@ -57,17 +57,23 @@ void ULiveLinkCameraController::Tick(float DeltaTime, const FLiveLinkSubjectFram
 
 			if (UCineCameraComponent* CineCameraComponent = Cast<UCineCameraComponent>(CameraComponent))
 			{
-				// If depth of field is not supported, force disable depth of field on the camera
-				// Otherwise, if depth of field is supported and the frame data has a valid value for focus, force the focus method to Manual 
-				if (UpdateFlags.bApplyFocusDistance)
+				// If the depth of field support flag has changed, update the focus method accordingly
+				if (bLastIsDepthOfFieldSupported != StaticData->bIsDepthOfFieldSupported)
 				{
-					if (StaticData->bIsDepthOfFieldSupported == false)
+					// If depth of field is not supported, force disable depth of field on the camera
+					// Otherwise, if depth of field is supported and the frame data has a valid value for focus, force the focus method to Manual 
+					if (UpdateFlags.bApplyFocusDistance)
 					{
-						CineCameraComponent->FocusSettings.FocusMethod = ECameraFocusMethod::Disable;
-					}
-					else if (StaticData->bIsFocusDistanceSupported)
-					{
-						CineCameraComponent->FocusSettings.FocusMethod = ECameraFocusMethod::Manual;
+						if (StaticData->bIsDepthOfFieldSupported == false)
+						{
+							CineCameraComponent->FocusSettings.FocusMethod = ECameraFocusMethod::Disable;
+						}
+						else if (StaticData->bIsFocusDistanceSupported)
+						{
+							CineCameraComponent->FocusSettings.FocusMethod = ECameraFocusMethod::Manual;
+						}
+
+						bLastIsDepthOfFieldSupported = StaticData->bIsDepthOfFieldSupported;
 					}
 				}
 
diff --git a/Engine/Plugins/VirtualProduction/LiveLinkCamera/Source/LiveLinkCamera/Public/LiveLinkCameraController.h b/Engine/Plugins/VirtualProduction/LiveLinkCamera/Source/LiveLinkCamera/Public/LiveLinkCameraController.h
--- a/Engine/Plugins/VirtualProduction/LiveLinkCamera/Source/LiveLinkCamera/Public/LiveLinkCameraController.h
+++ b/Engine/Plugins/VirtualProduction/LiveLinkCamera/Source/LiveLinkCamera/Public/LiveLinkCameraController.h
@@ -173,4 +173,6 @@ protected:
 
 private:
 	double LastLensTableVerificationTimestamp = 0.0;
+
+	bool bLastIsDepthOfFieldSupported = false;
 };
